Merging Hand and keyboard Fixes for HL2 (CL 14386402)

diff --git a/Engine/Plugins/Runtime/WindowsMixedReality/Source/WindowsMixedRealityHMD/Private/WindowsMixedRealityHMD.cpp b/Engine/Plugins/Runtime/WindowsMixedReality/Source/WindowsMixedRealityHMD/Private/WindowsMixedRealityHMD.cpp
--- a/Engine/Plugins/Runtime/WindowsMixedReality/Source/WindowsMixedRealityHMD/Private/WindowsMixedRealityHMD.cpp
+++ b/Engine/Plugins/Runtime/WindowsMixedReality/Source/WindowsMixedRealityHMD/Private/WindowsMixedRealityHMD.cpp
@@ -874,6 +874,7 @@ namespace WindowsMixedReality
 		MotionControllerData.DeviceName = GetSystemName();
 		MotionControllerData.ApplicationInstanceID = FApp::GetInstanceId();
 		MotionControllerData.DeviceVisualType = EXRVisualType::Hand;
+		MotionControllerData.HandIndex = Hand;
 
 		MotionControllerData.TrackingStatus = (ETrackingStatus)GetControllerTrackingStatus((WindowsMixedReality::HMDHand)Hand);
 
@@ -1239,8 +1240,16 @@ namespace WindowsMixedReality
 		return true;
 	}
 
+	bool FWindowsMixedRealityHMD::bEnableStereoReentranceGuard = false;
+
 	bool FWindowsMixedRealityHMD::EnableStereo(bool stereo)
 	{
+		if (bEnableStereoReentranceGuard)
+		{
+			return false;
+		}
+
+		bEnableStereoReentranceGuard = true;
 #if WITH_WINDOWS_MIXED_REALITY
 		if (stereo)
 		{
@@ -1251,18 +1260,20 @@ namespace WindowsMixedReality
 				if (FParse::Value(FCommandLine::Get(), TEXT("RemotingAppIp="), RemotingAppIp) &&
 					FParse::Value(FCommandLine::Get(), TEXT("RemotingBitRate="), RemotingBitRate))
 				{
-					ConnectToRemoteHoloLens(*RemotingAppIp, RemotingBitRate, false);
+					FWindowsMixedRealityStatics::ConnectToRemoteHoloLens(*RemotingAppIp, RemotingBitRate, false);
 				}
 			}
 
 			if (bIsStereoDesired && HMD->IsInitialized())
 			{
+				bEnableStereoReentranceGuard = false;
 				return false;
 			}
 
 			FindMRSceneViewport(bIsStereoDesired);
 			if (!bIsStereoDesired)
 			{
+				bEnableStereoReentranceGuard = false;
 				return false;
 			}
 
@@ -1336,6 +1347,8 @@ namespace WindowsMixedReality
 			StopSpeechRecognition();
 		}
 #endif
+		bEnableStereoReentranceGuard = false;
+
 		return bIsStereoDesired;
 	}
 
diff --git a/Engine/Plugins/Runtime/WindowsMixedReality/Source/WindowsMixedRealityHMD/Private/WindowsMixedRealityHMD.h b/Engine/Plugins/Runtime/WindowsMixedReality/Source/WindowsMixedRealityHMD/Private/WindowsMixedRealityHMD.h
--- a/Engine/Plugins/Runtime/WindowsMixedReality/Source/WindowsMixedRealityHMD/Private/WindowsMixedRealityHMD.h
+++ b/Engine/Plugins/Runtime/WindowsMixedReality/Source/WindowsMixedRealityHMD/Private/WindowsMixedRealityHMD.h
@@ -429,5 +429,8 @@ namespace WindowsMixedReality
 
 	private:
 		void CreateSpectatorScreenController();
+
+		//At runtime, EnableStereo can call FWindowsMixedRealityStatics::ConnectToRemoteHoloLens, but in the editor FWindowsMixedRealityStatics::ConnectToRemoteHoloLens calls EnableStereo.  This is to prevent a stack overflow
+		static bool bEnableStereoReentranceGuard;
 	};
 }
diff --git a/Engine/Plugins/Runtime/XR/XRVisualization/Source/XRVisualization/Private/XRVisualizationFunctionLibrary.cpp b/Engine/Plugins/Runtime/XR/XRVisualization/Source/XRVisualization/Private/XRVisualizationFunctionLibrary.cpp
--- a/Engine/Plugins/Runtime/XR/XRVisualization/Source/XRVisualization/Private/XRVisualizationFunctionLibrary.cpp
+++ b/Engine/Plugins/Runtime/XR/XRVisualization/Source/XRVisualization/Private/XRVisualizationFunctionLibrary.cpp
@@ -11,6 +11,8 @@
 #include "Engine/StaticMesh.h"
 #include "TimerManager.h"
 #include "UObject/ConstructorHelpers.h"
+#include "ProceduralMeshComponent.h"
+#include "IHandTracker.h"
 
 
 UXRVisualizationLoadHelper::UXRVisualizationLoadHelper(const FObjectInitializer& ObjectInitializer)
@@ -91,7 +93,7 @@ void UXRVisualizationFunctionLibrary::RenderMotionController(const FXRMotionCont
 
 			FTransform WorldTransform(XRControllerData.GripRotation, XRControllerData.GripPosition, bRight ? FVector(1.0f, -1.0f, 1.0f) : FVector(1.0f, 1.0f, 1.0f));
 
-			FName ActorName(*FString::Printf(TEXT("XR_%s_%x"), *XRControllerData.DeviceName.ToString(), (PTRINT)&XRControllerData));
+			FName ActorName(*FString::Printf(TEXT("XR_%s_%x_%d"), *XRControllerData.DeviceName.ToString(), (PTRINT)&XRControllerData), (int)XRControllerData.HandIndex);
 
 			UStaticMesh* MeshToUse = nullptr;
 			if (XRControllerData.DeviceName == TEXT("OculusHMD"))
@@ -169,11 +171,81 @@ void UXRVisualizationFunctionLibrary::RenderGenericMesh(const FName& ActorName,
 
 void UXRVisualizationFunctionLibrary::RenderHandMesh(const FXRMotionControllerData& XRControllerData)
 {
-	//if (XRVisualizationData.ApplicationInstanceID == FApp::GetInstanceId())
-	//{
-	//	//ASK TRACKING SYSTEM FOR SPECIFIC MESH IT MIGHT PROVIDE
-	//}
-	//else
+	bool bLocalRender = false;
+
+	if (XRControllerData.ApplicationInstanceID == FApp::GetInstanceId())
+	{
+		TArray<IHandTracker*> HandTrackers = IModularFeatures::Get().GetModularFeatureImplementations<IHandTracker>(IHandTracker::GetModularFeatureName());
+		IHandTracker* HandTracker = nullptr;
+		for (auto Itr : HandTrackers)
+		{
+			if (Itr->IsHandTrackingStateValid() && Itr->HasHandMeshData())
+			{
+				HandTracker = Itr;
+				break;
+			}
+		}
+
+		if (HandTracker != nullptr)
+		{
+			int32 HandIndex = (int)XRControllerData.HandIndex;
+			if ((HandIndex < 0) || (HandIndex >= 2))
+			{
+				return;
+			}
+
+			TArray<struct FVector> Vertices, Normals;
+			TArray<int32> Indices;
+			FTransform HandMeshTransform;
+
+			if ((HandTracker->GetHandMeshData((EControllerHand)XRControllerData.HandIndex, Vertices, Normals, Indices, HandMeshTransform)) && (Vertices.Num() > 0))
+			{
+				bLocalRender = true;
+
+				FName ActorName(*FString::Printf(TEXT("XR_hand_mesh_%s_%x_%d"), *XRControllerData.DeviceName.ToString(), (PTRINT)&XRControllerData, XRControllerData.HandIndex));
+				ULevel* CurrentLevel = GWorld->GetCurrentLevel();
+
+				AActor* FoundActor = FindObjectFast<AActor>(CurrentLevel, ActorName, true);
+				UProceduralMeshComponent* ProceduralMeshComponent = nullptr;
+				if ((FoundActor == nullptr) || (FoundActor->IsPendingKill()))
+				{
+					//create actor
+					FoundActor = NewObject<AActor>(CurrentLevel, ActorName);
+					//create component
+					ProceduralMeshComponent = NewObject<UProceduralMeshComponent>(FoundActor, UProceduralMeshComponent::StaticClass());
+					ProceduralMeshComponent->SetCollisionProfileName(UCollisionProfile::NoCollision_ProfileName);
+					ProceduralMeshComponent->RegisterComponentWithWorld(GWorld);
+				}
+				else
+				{
+					//get static mesh component out of the actor
+					ProceduralMeshComponent = Cast<UProceduralMeshComponent>(FoundActor->GetComponentByClass(UProceduralMeshComponent::StaticClass()));
+				}
+
+				if (ProceduralMeshComponent)
+				{
+					TArray<FLinearColor> Colors(&FLinearColor::White, Vertices.Num());
+
+					//Upload the data
+					ProceduralMeshComponent->CreateMeshSection_LinearColor(0, Vertices, Indices, Normals, {}, Colors, {}, false);
+
+					//Update the transform
+					ProceduralMeshComponent->SetWorldTransform(HandMeshTransform);
+
+					FTimerManager& TimerManager = FoundActor->GetWorldTimerManager();
+					//reset this timer and start over
+					TimerManager.ClearAllTimersForObject(FoundActor);
+
+					//Add timer to turn the component back off
+					FTimerHandle TempHandle;
+					TimerManager.SetTimer(TempHandle, FoundActor, &AActor::K2_DestroyActor, 2.0f, true);
+				}
+			}
+		}
+	}
+
+	//in case no hand tracker provided rendering 
+	if (!bLocalRender)
 	{
 		RenderFinger(XRControllerData, (int32)EHandKeypoint::ThumbMetacarpal, (int32)EHandKeypoint::ThumbTip);
 		RenderFinger(XRControllerData, (int32)EHandKeypoint::IndexMetacarpal, (int32)EHandKeypoint::IndexTip);
diff --git a/Engine/Plugins/Runtime/XR/XRVisualization/Source/XRVisualization/XRVisualization.Build.cs b/Engine/Plugins/Runtime/XR/XRVisualization/Source/XRVisualization/XRVisualization.Build.cs
--- a/Engine/Plugins/Runtime/XR/XRVisualization/Source/XRVisualization/XRVisualization.Build.cs
+++ b/Engine/Plugins/Runtime/XR/XRVisualization/Source/XRVisualization/XRVisualization.Build.cs
@@ -15,7 +15,8 @@ namespace UnrealBuildTool.Rules
                     "EngineSettings",
                     "RenderCore",
 					"RHI",
-					"HeadMountedDisplay"
+					"HeadMountedDisplay",
+					"ProceduralMeshComponent",
 				}
 			);
         }
diff --git a/Engine/Plugins/Runtime/XR/XRVisualization/XRVisualization.uplugin b/Engine/Plugins/Runtime/XR/XRVisualization/XRVisualization.uplugin
--- a/Engine/Plugins/Runtime/XR/XRVisualization/XRVisualization.uplugin
+++ b/Engine/Plugins/Runtime/XR/XRVisualization/XRVisualization.uplugin
@@ -1,24 +1,30 @@
 {
-	"FileVersion" : 3,
-	"Version" : 1,
-	"VersionName" : "1.0",
-	"FriendlyName" : "XRVisualization",
-	"Description" : "Visualization Library for XR HMDs and controllers",
-	"Category" : "Mixed Reality",
-	"CreatedBy" : "Epic Games",
-	"CreatedByURL" : "epicgames.com",
-	"DocsURL" : "",
-	"MarketplaceURL" : "",
-	"SupportURL" : "",
-	"EnabledByDefault" : false,
-	"CanContainContent" : false,
-	"IsExperimentalVersion" : true,
-	"IsBetaVersion" : true,
-	"Installed" : true,
+	"FileVersion": 3,
+	"Version": 1,
+	"VersionName": "1.0",
+	"FriendlyName": "XRVisualization",
+	"Description": "Visualization Library for XR HMDs and controllers",
+	"Category": "Mixed Reality",
+	"CreatedBy": "Epic Games",
+	"CreatedByURL": "epicgames.com",
+	"DocsURL": "",
+	"MarketplaceURL": "",
+	"SupportURL": "",
+	"EnabledByDefault": false,
+	"CanContainContent": false,
+	"IsExperimentalVersion": true,
+	"IsBetaVersion": true,
+	"Installed": true,
 	"Modules": [
 		{
 			"Name": "XRVisualization",
 			"Type": "Runtime"
 		}
+	],
+	"Plugins": [
+		{
+			"Name": "ProceduralMeshComponent",
+			"Enabled": true
+		}
 	]
 }
diff --git a/Engine/Source/Runtime/HeadMountedDisplay/Private/XRTrackingSystemBase.cpp b/Engine/Source/Runtime/HeadMountedDisplay/Private/XRTrackingSystemBase.cpp
--- a/Engine/Source/Runtime/HeadMountedDisplay/Private/XRTrackingSystemBase.cpp
+++ b/Engine/Source/Runtime/HeadMountedDisplay/Private/XRTrackingSystemBase.cpp
@@ -132,6 +132,7 @@ void FXRTrackingSystemBase::GetMotionControllerData(UObject* WorldContext, const
 	GetObjectsOfClass(UMotionControllerComponent::StaticClass(), MotionControllers);
 
 	MotionControllerData.bValid = false;
+	MotionControllerData.HandIndex = Hand;
 
 	for (int32 MotionControllerIndex = 0; MotionControllerIndex < MotionControllers.Num(); ++MotionControllerIndex)
 	{
diff --git a/Engine/Source/Runtime/HeadMountedDisplay/Public/HeadMountedDisplayTypes.h b/Engine/Source/Runtime/HeadMountedDisplay/Public/HeadMountedDisplayTypes.h
--- a/Engine/Source/Runtime/HeadMountedDisplay/Public/HeadMountedDisplayTypes.h
+++ b/Engine/Source/Runtime/HeadMountedDisplay/Public/HeadMountedDisplayTypes.h
@@ -311,6 +311,9 @@ struct HEADMOUNTEDDISPLAY_API FXRMotionControllerData
 	UPROPERTY(BlueprintReadOnly, Category = "XR")
 	EXRVisualType DeviceVisualType;
 
+	UPROPERTY(BlueprintReadOnly, Category = "XR")
+	EControllerHand HandIndex;
+
 	UPROPERTY(BlueprintReadOnly, Category = "XR")
 	ETrackingStatus TrackingStatus = ETrackingStatus::NotTracked;
 	
diff --git a/Engine/Source/Runtime/HeadMountedDisplay/Public/IHandTracker.h b/Engine/Source/Runtime/HeadMountedDisplay/Public/IHandTracker.h
--- a/Engine/Source/Runtime/HeadMountedDisplay/Public/IHandTracker.h
+++ b/Engine/Source/Runtime/HeadMountedDisplay/Public/IHandTracker.h
@@ -44,4 +44,14 @@ public:
 	virtual bool GetKeypointState(EControllerHand Hand, EHandKeypoint Keypoint, FTransform& OutTransform, float& OutRadius) const = 0;
 
 	virtual bool GetAllKeypointStates(EControllerHand Hand, TArray<struct FVector>& OutPositions, TArray<struct FQuat>& OutRotations, TArray<float>& OutRadii) const = 0;
+
+	virtual bool HasHandMeshData() const
+	{
+		return false;
+	}
+
+	virtual bool GetHandMeshData(EControllerHand Hand, TArray<struct FVector>& OutVertices, TArray<struct FVector>& OutNormals, TArray<int32>& OutIndices, FTransform& OutHandMeshTransform) const
+	{
+		return false;
+	}
 };
