Make UXRVisualizationFunctionLibrary methods accessible for C++

diff --git a/Engine/Plugins/Runtime/XR/XRVisualization/Source/XRVisualization/Public/XRVisualizationFunctionLibrary.h b/Engine/Plugins/Runtime/XR/XRVisualization/Source/XRVisualization/Public/XRVisualizationFunctionLibrary.h
--- a/Engine/Plugins/Runtime/XR/XRVisualization/Source/XRVisualization/Public/XRVisualizationFunctionLibrary.h
+++ b/Engine/Plugins/Runtime/XR/XRVisualization/Source/XRVisualization/Public/XRVisualizationFunctionLibrary.h
@@ -34,6 +34,8 @@ class XRVISUALIZATION_API UXRVisualizationFunctionLibrary : public UBlueprintFun
 
 	virtual ~UXRVisualizationFunctionLibrary();
 
+	// These methods should be visible from C++.
+public:
 	/**
 	 * Renders HMD Visualization on a system that might not have that HMD as native
 	 */
