Move "installed" plugin XRVisualization to "built-in".

diff --git a/Engine/Plugins/Runtime/XR/XRVisualization/XRVisualization.uplugin b/Engine/Plugins/Runtime/XR/XRVisualization/XRVisualization.uplugin
--- a/Engine/Plugins/Runtime/XR/XRVisualization/XRVisualization.uplugin
+++ b/Engine/Plugins/Runtime/XR/XRVisualization/XRVisualization.uplugin
@@ -14,7 +14,7 @@
 	"CanContainContent": false,
 	"IsExperimentalVersion": true,
 	"IsBetaVersion": true,
-	"Installed": true,
+	"Installed": false,
 	"Modules": [
 		{
 			"Name": "XRVisualization",
